Create ICVFXTesting plugin with the intention of porting SaloonPerf to use it

diff --git a/Engine/Plugins/VirtualProduction/ICVFX/ICVFX.uplugin b/Engine/Plugins/VirtualProduction/ICVFX/ICVFX.uplugin
--- a/Engine/Plugins/VirtualProduction/ICVFX/ICVFX.uplugin
+++ b/Engine/Plugins/VirtualProduction/ICVFX/ICVFX.uplugin
@@ -125,6 +125,10 @@
 		{
 			"Name": "RivermaxMedia",
 			"Enabled": true
+		},
+		{
+			"Name": "ICVFXTesting",
+			"Enabled": true
 		}
 	]
 }
diff --git a/Engine/Plugins/VirtualProduction/ICVFXTesting/ICVFXTesting.uplugin b/Engine/Plugins/VirtualProduction/ICVFXTesting/ICVFXTesting.uplugin
new file mode 100644
--- /dev/null
+++ b/Engine/Plugins/VirtualProduction/ICVFXTesting/ICVFXTesting.uplugin
@@ -0,0 +1,30 @@
+{
+	"FileVersion": 3,
+	"Version": 1,
+	"VersionName": "1.0",
+	"FriendlyName": "ICVFXTesting",
+	"Description": "Testing utilities to be used by ICVFX projects",
+	"Category": "Virtual Production",
+	"CreatedBy": "Epic Games, Inc.",
+	"CreatedByURL": "https://epicgames.com",
+	"DocsURL": "",
+	"MarketplaceURL": "",
+	"SupportURL": "",
+	"CanContainContent": false,
+	"IsBetaVersion": true,
+	"Hidden": true,
+	"Installed": false,
+	"Plugins": [
+		{
+			"Name": "Gauntlet",
+			"Enabled": true
+		}
+	],
+	"Modules": [
+		{
+			"Name": "ICVFXTesting",
+			"Type": "Runtime",
+			"LoadingPhase": "Default"
+		}
+	]
+}
diff --git a/Engine/Plugins/VirtualProduction/ICVFXTesting/Scripts/GauntletSettings.xml b/Engine/Plugins/VirtualProduction/ICVFXTesting/Scripts/GauntletSettings.xml
new file mode 100644
--- /dev/null
+++ b/Engine/Plugins/VirtualProduction/ICVFXTesting/Scripts/GauntletSettings.xml
@@ -0,0 +1,100 @@
+<?xml version='1.0' ?>
+<BuildGraph xmlns="http://www.epicgames.com/BuildGraph" xmlns:xsi="http://www.w3.org/2001/XMLSchema-instance" xsi:schemaLocation="http://www.epicgames.com/BuildGraph ../../Schema.xsd" >
+
+	<Property Name="ProjectTestName" Value="$(ProjectName)Test"/>
+		
+	<Option Name="WithEditorAutoTest" Restrict="true|false" DefaultValue="false" Description="Whether to run the AutoTest on the editor build." />
+	<Option Name="WithEditorAutomationPIETest" Restrict="true|false" DefaultValue="false" Description="Whether to run UE.EditorAutomation(RunTest=Project.Maps.PIE)." />
+	<Option Name="WithEditorAutomationAll" Restrict="true|false" DefaultValue="false" Description="Whether to run UE.EditorAutomation(RunTest=Project.)." />
+
+	<Option Name="WithAutoTest" Restrict="true|false" DefaultValue="false" Description="Whether to run AutoTest on the target builds." />
+	<Option Name="WithPerformanceReportTest" Restrict="true|false" DefaultValue="true" Description="Whether to run the PerformanceReport tests." />
+	<Option Name="WithMemoryReportTest" Restrict="true|false" DefaultValue="false" Description="Whether to run the MemoryReport test." />
+	<Option Name="WithTargetAutomationMapsCycle" Restrict="true|false" DefaultValue="false" Description="Whether to run UE.TargetAutomation(RunTest=Project.Maps.Cycle)." />
+	<Option Name="WithTargetAutomationAll" Restrict="true|false" DefaultValue="false" Description="Whether to run UE.TargetAutomation(RunTest=Project.)." />
+
+	<!-- Hack solution to set target test list for BuildAndTest.xml -->
+	<Do If="'$(EditorTestList)' == ''">
+		<!-- <Property Name="EditorTestList" Value="CitySampleTest.AutoTest+" If="$(WithEditorAutoTest)" />
+		<Property Name="EditorTestList" Value="$(EditorTestList)UE.EditorAutomation(RunTest=Project.)+" If="$(WithEditorAutomationAll)" />
+		<Property Name="EditorTestList" Value="$(EditorTestList)UE.EditorAutomation(RunTest=Project.Maps.PIE)" If="!$(WithEditorAutomationAll) and $(WithEditorAutomationPIETest)" />-->
+		<Property Name="TargetTestList" Value="$(EditorTestList)$(ProjectTestName).MemoryReport+" If="$(WithMemoryReportTest)" />
+		<Property Name="EditorTestList" Value="$(EditorTestList)$(ProjectTestName).PerformanceReport+" If="$(WithPerformanceReportTest)" />
+		<Property Name="EditorTestList" Value="$(EditorTestList)$(ProjectTestName).PerformanceReport_MGPU+" If="$(WithPerformanceReportTest)" />
+		<Property Name="EditorTestList" Value="$(EditorTestList)$(ProjectTestName).PerformanceReport_Nanite+" If="$(WithPerformanceReportTest)" />
+		<Property Name="EditorTestList" Value="$(EditorTestList)$(ProjectTestName).PerformanceReport_NaniteLumen+" If="$(WithPerformanceReportTest)" />
+		<Property Name="EditorTestList" Value="$(EditorTestList)$(ProjectTestName).PerformanceReport_Vulkan+" If="$(WithPerformanceReportTest)" />
+	</Do>
+
+	<!-- How long to give Gauntlet before a test is considered timed out (some tests will time out earlier if they detect inactivity)-->
+	<Option Name="GauntletTimeout" DefaultValue="900" Description="Max running time of a test in Gauntlet"/>
+	
+	
+	<!-- Device option overrides -->
+	<Option Name="GauntletDevicePool" DefaultValue="UE5" Description="Gauntlet Device Pool to use to run target platform tests"/>
+	<Option Name="AlternateDevicePoolPlatforms" DefaultValue="" Description="Platforms that need to use the Alternate Gauntlet Device Pool"/>
+	
+	<!-- Test specific options -->
+	<Option Name="DisplayConfigPath" DefaultValue="" Description="The path to the .ndisplay file used for launching the cluster."/>
+	<Option Name="DisplayClusterNodeName" DefaultValue="node_0" Description="The name of the ndisplay node to launch."/>
+	<Option Name="ViewportScreenPercentage" DefaultValue="100.0" Description="The viewport screen percentage to use"/>
+	<Option Name="WithNanite" Restrict="true|false" DefaultValue="false" Description="Whether to use nanite when testing"/>
+	<Option Name="WithLumen" Restrict="true|false" DefaultValue="false" Description="Whether to use lumen when testing"/>
+	<Option Name="WithD3DDebug" Restrict="true|false" DefaultValue="false" Description="Whether to add d3ddebug command line argument."/>
+	<Option Name="WithGPUCrashDebugging" Restrict="true|false" DefaultValue="false" Description="Whether to add gpucrashdebugging  command line argument."/>
+	<Option Name="MaxGPUCount" DefaultValue="1" Description="Number of GPUs to use for the test."/>
+	<Option Name="MaxRunCount" DefaultValue="2" Description="Number of iterations to do."/>
+	<Option Name="WithRDGImmediate" Restrict="true|false" DefaultValue="false" Description="Whether to run with -rdgimmediate"/>
+	<Option Name="WithNoRHIThread" Restrict="true|false" DefaultValue="false" Description="Whether to run with -norhithread"/>
+	<Option Name="WithVirtualShadowMaps" Restrict="true|false" DefaultValue="true" Description="Whether to use virtual shadow maps"/>
+	<Option Name="WithTrace" Restrict="true|false" DefaultValue="false" Description="Whether to run with unreal insights."/>
+	
+	
+	<Property Name="Platform" Value="Win64"/>
+
+	<!-- Automation uses its own network output directory -->
+	<Property Name="AutomationNetworkRootDir" Value="\\epicgames.net\root\Builds\Automation\Reports\$(ProjectName)Perf"/>
+	<Property Name="AutomationNetworkRootDir" Value="/Volumes/Builds/Automation/Reports/$(ProjectName)Perf" If="'$(HostPlatform)' != 'Win64'"/>
+	<Property Name="AutomationLocalOutputDir" Value="$(RootDir)/LocalBuilds/Reports"/>
+
+	<!-- Outputs for this run keeping slashes flipped the right way because it shows up on the webapp this way -->
+	<Property Name="AutomationNetworkOutputDir" Value="$(AutomationNetworkRootDir)\$(BuildName)"/>
+	<Property Name="AutomationNetworkOutputDir" Value="$(AutomationNetworkRootDir)/$(BuildName)" If="'$(HostPlatform)' != 'Win64'"/>
+	<Property Name="AutomationOutputDir" Value="$(AutomationLocalOutputDir)"/>
+	<Property Name="AutomationOutputDir" Value="$(AutomationNetworkOutputDir)" If="$(IsBuildMachine)"/>
+	
+	<Property Name="SummaryReportPath" Value="$(AutomationLocalOutputDir)"/>
+	<Property Name="SummaryReportPath" Value="$(AutomationNetworkRootDir)" If="$(IsBuildMachine)"/>
+	
+	
+	<Property Name="TestOutputDir" Value="$(AutomationOutputDir)/$(ProjectName)$(Platform)$(ProjectName)Perf"/>
+	<Property Name="LogOutputParams" Value="-logdir=&quot;$(TestOutputDir)&quot;"/>
+	<!-- <Property Name="ReportParams" Value="-WriteTestResultsForHorde -HordeTestDataKey=&quot;$(NodeName)&quot;"/> -->
+	<!--<Property Name="ReportParams" Value="-WriteTestResultsForHorde"/>-->
+	<!--<Property Name="ArtifactName" Value="$(Phase)$(IterationSuffix)$(TestNameSuffix)"/>-->
+	<!--<Property Name="AllGauntletArgs" Value="-Project=$(Project) -Platform=$(Platform) $(GauntletArgs) $(ReportParams) $(LogOutputParams) -artifactname=&quot;$(ArtifactName)&quot;"/>-->
+	<!--<Property Name="AllGauntletArgs" Value="$(LogOutputParams)"/>-->
+	
+		
+	<Property Name="ExtraGauntletCommonArgs" Value="-screenshotperiod=30 -$(ProjectTestName).StatCommands -$(ProjectTestName).MaxRunCount=$(MaxRunCount) $(LogOutputParams)" />
+	<Property Name="ExtraGauntletEditorTestArgs" Value="-game" />
+	<Property Name="ExtraGauntletEditorTestArgs" Value="$(ExtraGauntletEditorTestArgs) -$(ProjectTestName).MapOverride=$(MapOverride) -$(ProjectTestName).DisplayConfig=$(DisplayConfigPath) -$(ProjectTestName).ViewportScreenPercentage=$(ViewportScreenPercentage)" />
+	<Property Name="ExtraGauntletEditorTestArgs" Value="$(ExtraGauntletEditorTestArgs) -$(ProjectTestName).DisplayClusterNodeName=$(DisplayClusterNodeName)" />
+	
+	<Property Name="ExtraGauntletEditorTestArgs" Value="$(ExtraGauntletEditorTestArgs) -$(ProjectTestName).D3DDebug" If="$(WithD3DDebug)"/>
+	<Property Name="ExtraGauntletEditorTestArgs" Value="$(ExtraGauntletEditorTestArgs) -$(ProjectTestName).GPUCrashDebugging" If="$(WithGPUCrashDebugging)"/>
+	<Property Name="ExtraGauntletEditorTestArgs" Value="$(ExtraGauntletEditorTestArgs) -$(ProjectTestName).Nanite=$(WithNanite)"/>
+	<Property Name="ExtraGauntletEditorTestArgs" Value="$(ExtraGauntletEditorTestArgs) -$(ProjectTestName).Lumen=$(WithLumen)"/>
+	<Property Name="ExtraGauntletEditorTestArgs" Value="$(ExtraGauntletEditorTestArgs) -$(ProjectTestName).MaxGPUCount=$(MaxGPUCount)"/>
+	<Property Name="ExtraGauntletEditorTestArgs" Value="$(ExtraGauntletEditorTestArgs) -$(ProjectTestName).SummaryReportPath=$(SummaryReportPath)"/>
+	<Property Name="ExtraGauntletEditorTestArgs" Value="$(ExtraGauntletEditorTestArgs) -$(ProjectTestName).RDGImmediate" If="$(WithRDGImmediate)"/>
+	<Property Name="ExtraGauntletEditorTestArgs" Value="$(ExtraGauntletEditorTestArgs) -$(ProjectTestName).NoRHIThread" If="$(WithNoRHIThread)"/>
+	<Property Name="ExtraGauntletEditorTestArgs" Value="$(ExtraGauntletEditorTestArgs) -$(ProjectTestName).DisableVirtualShadowMaps" If="!$(WithVirtualShadowMaps)"/>
+	<Property Name="ExtraGauntletEditorTestArgs" Value="$(ExtraGauntletEditorTestArgs) -$(ProjectTestName).EnableTrace" If="$(WithTrace)"/>
+	<Property Name="ExtraGauntletEditorTestArgs" Value="$(ExtraGauntletEditorTestArgs) -$(ProjectTestName).TraceRootFolder=$(TestOutputDir)/$(BuildName)" If="$(WithTrace)"/>
+	
+	
+	<!-- Common Gauntlet settings -->
+	<Include Script="../../../../../Engine/Build/Graph/Tasks/Inc/GauntletSettings.xml" />
+
+</BuildGraph>
\ No newline at end of file
diff --git a/Engine/Plugins/VirtualProduction/ICVFXTesting/Scripts/ICVFXTesting.xml b/Engine/Plugins/VirtualProduction/ICVFXTesting/Scripts/ICVFXTesting.xml
new file mode 100644
--- /dev/null
+++ b/Engine/Plugins/VirtualProduction/ICVFXTesting/Scripts/ICVFXTesting.xml
@@ -0,0 +1,42 @@
+<?xml version='1.0' ?>
+<BuildGraph xmlns="http://www.epicgames.com/BuildGraph" xmlns:xsi="http://www.w3.org/2001/XMLSchema-instance" xsi:schemaLocation="http://www.epicgames.com/BuildGraph ../Schema.xsd" >
+
+	<!-- todo can this be disbaled? -->
+	<Property Name="Versioned" Value="$(IsBuildMachine)" />
+	<Property Name="PublishHashedStagingData" Value="$(IsBuildMachine)" />
+
+	<!-- Turn off BuildAndTest defaults so we can set out own -->
+	<Property Name="WithBATDefaults" Value="false" />
+
+	<!-- Set defaults for people running this script with no arguments. These can still be overriden with set:EditorPlatforms= etc) -->
+	<Property Name="DefaultEditorPlatforms" Value="Win64" />
+	<Property Name="DefaultTargetPlatforms" Value="" /> <!-- We dont need a cooked build -->
+	<Property Name="DefaultTargetConfigurations" Value="Development" />
+	<Property Name="DefaultEditorTestList" Value="" />
+	<Property Name="DefaultTargetTestList" Value="" />
+	<Property Name="NetworkTempRootOverride" Value="" />
+	<Property Name="NetworkPublishRootOverride" Value="" />
+	<Property Name="NetworkReportRootOverride" Value="" />
+	<Property Name="SymbolServerPath" Value="" />
+
+	<Property Name="ExtraStageAndPackageArguments" Value="-compressed -CrashReporter" />
+	
+	<!-- If a build machine, turn off all defaults. Settings should be specified via the website UI -->
+	<Do If="$(IsBuildMachine)">
+		<Property Name="DefaultEditorPlatforms" Value="" />
+		<Property Name="DefaultTargetPlatforms" Value=""/>
+		<Property Name="DefaultTargetConfigurations" Value="Development"/>
+		<Property Name="DefaultEditorTestList" Value="" />
+		<Property Name="DefaultTargetTestList" Value="" />
+	</Do>
+
+	<Property Name="AgentOverride" Value="EditorPerfWin64" />
+	<Property Name="TestAgentOverride" Value="EditorPerfWin64" />
+	<Property Name="ClientCompileAgentOverride" Value="EditorPerfWin64" />
+	<!-- Common Setting Overrides -->
+	
+	<Property Name="GauntletSettingsFile" Value="$(RootDir)/Engine/Plugins/VirtualProduction/ICVFX/Scripts/GauntletSettings.xml" />
+
+	<!-- This will declare an aggregate called BuildAndTest $(ProjectName) -->
+	<Include Script="$(RootDir)/Engine/Build/Graph/Tasks/BuildAndTestProject.xml" />
+</BuildGraph>
\ No newline at end of file
diff --git a/Engine/Plugins/VirtualProduction/ICVFXTesting/Source/ICVFXTesting/ICVFXTesting.build.cs b/Engine/Plugins/VirtualProduction/ICVFXTesting/Source/ICVFXTesting/ICVFXTesting.build.cs
new file mode 100644
--- /dev/null
+++ b/Engine/Plugins/VirtualProduction/ICVFXTesting/Source/ICVFXTesting/ICVFXTesting.build.cs
@@ -0,0 +1,24 @@
+// Copyright Epic Games, Inc. All Rights Reserved.
+
+using UnrealBuildTool;
+
+public class ICVFXTesting : ModuleRules
+{
+	public ICVFXTesting(ReadOnlyTargetRules Target) : base(Target)
+	{
+		PublicDependencyModuleNames.AddRange(
+			new string[] {
+				"Core",
+				"Engine",
+				"CoreUObject",
+				"Gauntlet"
+			}
+		);
+		
+		PrivateDependencyModuleNames.AddRange(
+			new string[] {
+            	"Engine"
+			}
+        );
+	}
+}
diff --git a/Engine/Plugins/VirtualProduction/ICVFXTesting/Source/ICVFXTesting/Private/ICVFXTestControllerAutoTest.cpp b/Engine/Plugins/VirtualProduction/ICVFXTesting/Source/ICVFXTesting/Private/ICVFXTestControllerAutoTest.cpp
new file mode 100644
--- /dev/null
+++ b/Engine/Plugins/VirtualProduction/ICVFXTesting/Source/ICVFXTesting/Private/ICVFXTestControllerAutoTest.cpp
@@ -0,0 +1,504 @@
+// Copyright Epic Games, Inc. All Rights Reserved.
+
+#include "ICVFXTestControllerAutoTest.h"
+
+#include "Engine/Engine.h"
+#include "Engine/LocalPlayer.h"
+#include "GameFramework/GameMode.h"
+#include "Logging/LogVerbosity.h"
+#include "PlatformFeatures.h"
+#include "ProfilingDebugging/CsvProfiler.h"
+#include "VideoRecordingSystem.h"
+
+namespace ICVFXTest
+{
+	static TAutoConsoleVariable<float> CVarSoakTime(
+		TEXT("ICVFXTest.SoakTime"),
+		30.0f,
+		TEXT("Duration of the sandbox soak, in seconds, per run."),
+		ECVF_Default);
+
+	static TAutoConsoleVariable<bool> CVarSkipTestSequence(
+		TEXT("ICVFXTest.SkipTestSequence"),
+		false,
+		TEXT("Whether to skip the test sequence after the sandbox soak, if applicable."),
+		ECVF_Default);
+
+	static TAutoConsoleVariable<bool> CVarRequestShutdown(
+		TEXT("ICVFXTest.RequestShutdown"),
+		false,
+		TEXT("Whether to request shutdown when the current run is finished, regardless of NumRuns."),
+		ECVF_Default);
+
+	static TAutoConsoleVariable<FString> CVarTraceFileName(
+		TEXT("ICVFXTest.TraceFileName"),
+		"",
+		TEXT("Path to the output tracefile."),
+		ECVF_Default);
+}
+
+//////////////////////////////////////////////////////////////////////
+// States
+//////////////////////////////////////////////////////////////////////
+
+// Base
+void FICVFXAutoTestState::Start(const EICVFXAutoTestState PrevState)
+{
+	check(IsValid(Controller));
+	TimeSinceStart = 0.0;
+}
+
+void FICVFXAutoTestState::Tick(const float TimeDelta)
+{
+	TimeSinceStart += TimeDelta;
+}
+
+// InitialLoad
+class FICVFXAutoTestState_InitialLoad : public FICVFXAutoTestState
+{
+public:
+	FICVFXAutoTestState_InitialLoad(UICVFXTestControllerAutoTest* const TestController) 
+		: FICVFXAutoTestState(TestController) 
+	{
+	}
+
+	virtual void Tick(float TimeDelta) override
+	{
+		Super::Tick(TimeDelta);
+
+		if (GetTestStateTime() >= MaxStateChangeWait)
+		{
+			UE_LOG(LogICVFXTest, Warning, TEXT("AutoTest InitialLoad %s: waited for state change for %f seconds, bindings may be missing, validate current state."), ANSI_TO_TCHAR(__func__), MaxStateChangeWait);
+			Controller->SetTestState(EICVFXAutoTestState::Soak);
+		}
+	}
+
+private:
+	const float MaxStateChangeWait = 30.0f;
+};
+
+// Sandbox Intro
+class FICVFXAutoTestState_Intro : public FICVFXAutoTestState
+{
+public:
+	FICVFXAutoTestState_Intro(UICVFXTestControllerAutoTest* const TestController) : FICVFXAutoTestState(TestController) {}
+
+	virtual void Start(const EICVFXAutoTestState PrevState) override
+	{
+		Super::Start(PrevState);
+
+		// If starting a new run, kick off requested captures
+
+		if (Controller->RequestsVideoCapture())
+		{
+			Controller->TryStartingVideoCapture();
+		}
+
+		if (Controller->RequestsFPSChart())
+		{
+			Controller->ConsoleCommand(TEXT("StartFPSChart"));
+		}
+	}
+};
+
+// Sandbox Soak
+class FICVFXAutoTestState_Soak : public FICVFXAutoTestState
+{
+public:
+	FICVFXAutoTestState_Soak(UICVFXTestControllerAutoTest* const TestController) : FICVFXAutoTestState(TestController) {}
+
+	virtual void Start(const EICVFXAutoTestState PrevState) override
+	{
+		Super::Start(PrevState);
+
+		// If starting a new run, kick off requested captures
+		if (PrevState == EICVFXAutoTestState::InitialLoad || PrevState == EICVFXAutoTestState::Finished)
+		{
+			if (Controller->RequestsVideoCapture())
+			{
+				Controller->TryStartingVideoCapture();
+			}
+
+			if (Controller->RequestsFPSChart())
+			{
+				Controller->ConsoleCommand(TEXT("StartFPSChart"));
+			}
+		}
+
+		const float SoakTime = ICVFXTest::CVarSoakTime->GetFloat();
+		if (SoakTime > 0.0f)
+		{
+			UE_LOG(LogICVFXTest, Display, TEXT("AutoTest Soak %s: soak time started for %f seconds..."), ANSI_TO_TCHAR(__func__), SoakTime);
+		}
+		else
+		{
+			UE_LOG(LogICVFXTest, Display, TEXT("AutoTest Soak %s: soaking indefinitely until soak time is set >0..."), ANSI_TO_TCHAR(__func__));
+		}
+	}
+
+	virtual void Tick(float TimeDelta) override
+	{
+		Super::Tick(TimeDelta);
+		
+		const float SoakTime = ICVFXTest::CVarSoakTime->GetFloat();
+		if (SoakTime > 0.0f && GetTestStateTime() >= SoakTime)
+		{
+			Controller->SetTestState(EICVFXAutoTestState::Idle);
+		}		
+	}
+};
+
+// Sandbox Soak
+class FICVFXAutoTestState_Idle : public FICVFXAutoTestState
+{
+public:
+	FICVFXAutoTestState_Idle(UICVFXTestControllerAutoTest* const TestController) : FICVFXAutoTestState(TestController) {}
+
+	const float IdleTime = 60.f;
+
+	virtual void Start(const EICVFXAutoTestState PrevState) override
+	{
+		Super::Start(PrevState);
+		if (!ICVFXTest::CVarTraceFileName.GetValueOnAnyThread().IsEmpty())
+		{
+			Controller->ConsoleCommand(*(FString(TEXT("trace.file ")) + ICVFXTest::CVarTraceFileName.GetValueOnAnyThread()));
+		}
+
+		UE_LOG(LogICVFXTest, Display, TEXT("AutoTest Idle %s: idle time started for %f seconds..."), ANSI_TO_TCHAR(__func__), IdleTime);
+	}
+
+	virtual void Tick(float TimeDelta) override
+	{
+		Super::Tick(TimeDelta);
+
+		if (IdleTime > 0.0f && GetTestStateTime() >= IdleTime)
+		{
+			Controller->SetTestState(EICVFXAutoTestState::Finished);
+		}
+	}
+};
+
+// Finished
+class FICVFXAutoTestState_Finished : public FICVFXAutoTestState
+{
+public:
+	FICVFXAutoTestState_Finished(UICVFXTestControllerAutoTest* const TestController) : FICVFXAutoTestState(TestController) {}
+
+	virtual void Start(const EICVFXAutoTestState PrevState) override
+	{
+		Super::Start(PrevState);
+		
+		if (!ICVFXTest::CVarTraceFileName.GetValueOnAnyThread().IsEmpty())
+		{
+			Controller->ConsoleCommand(TEXT("trace.stop"));
+		}
+
+		if (Controller->RequestsMemReport())
+		{
+			Controller->ExecuteMemReport();
+		}
+
+		Controller->MarkRunComplete();
+
+		// Success validation goes here...
+
+		bWaitingOnFPSChart = false;
+		bWaitingOnFinalizingVideo = false;
+
+#if CSV_PROFILER
+		if (FCsvProfiler* const CsvProfiler = FCsvProfiler::Get())
+		{
+			CsvProfilerDelegateHandle = CsvProfiler->OnCSVProfileFinished().AddLambda([this](const FString& Filename)
+			{
+				UE_LOG(LogICVFXTest, Display, TEXT("AutoTest OnCSVProfileFinished: CsvProfiler finished..."), ANSI_TO_TCHAR(__func__));
+				bWaitingOnFPSChart = false;
+			});
+
+			if (CsvProfiler->IsCapturing())
+			{
+				bWaitingOnFPSChart = true;
+				Controller->ConsoleCommand(TEXT("StopFPSChart"));
+			}
+			else
+			{
+				bWaitingOnFPSChart = CsvProfiler->IsWritingFile();
+			}
+
+		}
+#endif // CSV_PROFILER
+
+		if (IVideoRecordingSystem* const VideoRecordingSystem = IPlatformFeaturesModule::Get().GetVideoRecordingSystem())
+		{
+			VideoRecordingDelegateHandle = VideoRecordingSystem->GetOnVideoRecordingFinalizedDelegate().AddLambda([this](bool bSucceeded, const FString& FilePath)
+			{
+				if (bSucceeded)
+				{
+					UE_LOG(LogICVFXTest, Display, TEXT("AutoTest OnVideoRecordingFinalized: video recording successfully finalized at %s."), *FilePath);
+				}
+				else
+				{
+					UE_LOG(LogICVFXTest, Warning, TEXT("AutoTest OnVideoRecordingFinalized: video recording failed to finalize..."));
+				}
+
+				bWaitingOnFinalizingVideo = false;
+			});
+		}
+
+		bWaitingOnFinalizingVideo = Controller->TryFinalizingVideoCapture(true);
+	}
+
+	virtual void Tick(const float TimeDelta) override
+	{
+		Super::Tick(TimeDelta);
+
+		// Checks whether the test shutdown is complete and the test is ready to end
+		if (!bWaitingOnFPSChart && !bWaitingOnFinalizingVideo)
+		{
+			if (Controller->GetRunsRemaining() > 0 && !ICVFXTest::CVarRequestShutdown->GetBool())
+			{
+				const UWorld* const World = Controller->GetWorld();
+				if (AGameMode* const GameMode = World ? World->GetAuthGameMode<AGameMode>() : nullptr)
+				{
+					GameMode->RestartGame();
+				}
+			}
+			else
+			{
+				// If all runs are complete, then cleanup and shutdown test
+				Controller->SetTestState(EICVFXAutoTestState::Shutdown);
+			}
+		}
+	}
+
+	virtual void End(const EICVFXAutoTestState NewState) override
+	{
+#if CSV_PROFILER
+		if (FCsvProfiler* const CsvProfiler = FCsvProfiler::Get())
+		{
+			CsvProfiler->OnCSVProfileFinished().Remove(CsvProfilerDelegateHandle);
+		}
+#endif // CSV_PROFILER
+
+		if (IVideoRecordingSystem* const VideoRecordingSystem = IPlatformFeaturesModule::Get().GetVideoRecordingSystem())
+		{
+			VideoRecordingSystem->GetOnVideoRecordingFinalizedDelegate().Remove(VideoRecordingDelegateHandle);
+		}
+
+		Super::End(NewState);
+	}
+
+private:
+	bool bWaitingOnFPSChart = false;
+	bool bWaitingOnFinalizingVideo = false;
+
+	FDelegateHandle CsvProfilerDelegateHandle;
+	FDelegateHandle VideoRecordingDelegateHandle;
+};
+
+// Finished
+class FICVFXAutoTestState_Shutdown : public FICVFXAutoTestState
+{
+public:
+	FICVFXAutoTestState_Shutdown(UICVFXTestControllerAutoTest* const TestController) : FICVFXAutoTestState(TestController) {}
+
+	virtual void Start(const EICVFXAutoTestState PrevState) override
+	{
+		Super::Start(PrevState);
+		Controller->EndICVFXTest(TestExitCode);
+	}
+
+	void SetTestExitCode(const int32 ExitCode)
+	{
+		TestExitCode = ExitCode;
+	}
+
+private:
+	int32 TestExitCode = 0;
+};
+
+
+//////////////////////////////////////////////////////////////////////
+// Non state-specific logic
+//////////////////////////////////////////////////////////////////////
+
+
+FString UICVFXTestControllerAutoTest::GetStateName(const EICVFXAutoTestState State) const
+{
+	return StaticEnum<EICVFXAutoTestState>()->GetValueAsString(State);
+}
+
+FICVFXAutoTestState& UICVFXTestControllerAutoTest::GetTestState() const
+{
+	check(CurrentState != EICVFXAutoTestState::MAX);
+	FICVFXAutoTestState* State = States[(uint8)CurrentState];
+	check(State);
+	return *State;
+}
+
+FICVFXAutoTestState& UICVFXTestControllerAutoTest::SetTestState(const EICVFXAutoTestState NewState)
+{
+	const FString& EndingStateStr = GetStateName(CurrentState);
+	const FString& StartingStateStr = GetStateName(NewState);
+
+	if (CurrentState == NewState)
+	{
+		UE_LOG(LogICVFXTest, Display, TEXT("AutoTest %s: reloading %s"), ANSI_TO_TCHAR(__func__), *EndingStateStr);
+	}
+	else
+	{
+		UE_LOG(LogICVFXTest, Display, TEXT("AutoTest %s: %s -> %s"), ANSI_TO_TCHAR(__func__), *EndingStateStr, *StartingStateStr);
+	}
+
+	// End current test state
+	GetTestState().End(NewState);
+
+	// Transition current test state from the previous state to the new state
+	const EICVFXAutoTestState PrevState = CurrentState;
+	CurrentState = NewState;
+	FICVFXAutoTestState& CurrentTestState = GetTestState();
+
+	// Start new test state
+	CurrentTestState.Start(PrevState);
+	return CurrentTestState;
+}
+
+void UICVFXTestControllerAutoTest::EndICVFXTest(const int32 ExitCode /*= 0*/)
+{
+	if (CurrentState != EICVFXAutoTestState::Shutdown)
+	{
+		UE_LOG(LogICVFXTest, Warning, TEXT("AutoTest %s: called outside Shutdown test state, something went wrong."), ANSI_TO_TCHAR(__func__));
+		static_cast<FICVFXAutoTestState_Shutdown&>(SetTestState(EICVFXAutoTestState::Shutdown)).SetTestExitCode(ExitCode);
+		return;
+	}
+
+	Super::EndICVFXTest(ExitCode);
+}
+
+void UICVFXTestControllerAutoTest::UnbindAllDelegates()
+{
+	if (UWorld* const World = GetWorld())
+	{
+		World->OnWorldBeginPlay.RemoveAll(this);
+		World->GameStateSetEvent.RemoveAll(this);
+	}
+
+	IConsoleManager::Get().UnregisterConsoleVariableSink_Handle(SoakTimeSink);
+
+	Super::UnbindAllDelegates();
+}
+
+void UICVFXTestControllerAutoTest::OnInit()
+{
+	Super::OnInit();
+
+	// Initialize State Data
+	States[(uint8)EICVFXAutoTestState::InitialLoad] = new FICVFXAutoTestState_InitialLoad(this);
+	States[(uint8)EICVFXAutoTestState::Soak] = new FICVFXAutoTestState_Soak(this);
+	States[(uint8)EICVFXAutoTestState::Idle] = new FICVFXAutoTestState_Idle(this);
+	States[(uint8)EICVFXAutoTestState::Finished] = new FICVFXAutoTestState_Finished(this);
+	States[(uint8)EICVFXAutoTestState::Shutdown] = new FICVFXAutoTestState_Shutdown(this);
+
+	StateTimeouts[(uint8)EICVFXAutoTestState::InitialLoad] = 60.f;
+	StateTimeouts[(uint8)EICVFXAutoTestState::Soak] =  ICVFXTest::CVarSoakTime->GetFloat();
+	StateTimeouts[(uint8)EICVFXAutoTestState::Idle] = 60.0f;
+	StateTimeouts[(uint8)EICVFXAutoTestState::Finished] = 300.f;
+	StateTimeouts[(uint8)EICVFXAutoTestState::Shutdown] = 30.f;
+
+	for (int32 StateIdx = 0; StateIdx < (uint8)EICVFXAutoTestState::MAX; ++StateIdx)
+	{
+		checkf(States[StateIdx] != nullptr, TEXT("Missing state object for state %s"), *GetStateName((EICVFXAutoTestState)StateIdx));
+	}
+
+	const FConsoleCommandDelegate SoakTimeDelegate = FConsoleCommandDelegate::CreateUObject(this, &ThisClass::OnSoakTimeChanged);
+	SoakTimeSink = IConsoleManager::Get().RegisterConsoleVariableSink_Handle(SoakTimeDelegate);
+
+	// Set and start the InitialLoad test state
+	UE_LOG(LogICVFXTest, Display, TEXT("AutoTest %s: setting test state to InitialLoad"), ANSI_TO_TCHAR(__func__));
+	CurrentState = EICVFXAutoTestState::InitialLoad;
+	GetTestState().Start(EICVFXAutoTestState::MAX);
+}
+
+void UICVFXTestControllerAutoTest::OnPreMapChange()
+{
+	Super::OnPreMapChange();
+
+	if (CurrentState != EICVFXAutoTestState::InitialLoad)
+	{
+		UE_LOG(LogICVFXTest, Display, TEXT("AutoTest %s: setting test state to InitialLoad"), ANSI_TO_TCHAR(__func__));
+		SetTestState(EICVFXAutoTestState::InitialLoad);
+	}
+}
+
+void UICVFXTestControllerAutoTest::OnTick(float TimeDelta)
+{
+	Super::OnTick(TimeDelta);
+
+	GetTestState().Tick(TimeDelta);
+	MarkHeartbeatActive();
+
+	if (GetTestState().GetTestStateTime() > StateTimeouts[(uint8)CurrentState])
+	{
+		if (CurrentState == EICVFXAutoTestState::Finished || CurrentState == EICVFXAutoTestState::Shutdown)
+		{
+			// Treat timeouts during finalization/shutdown as warnings
+			UE_LOG(LogICVFXTest, Warning, TEXT("AutoTest %s: '%s' test state timed out after %f seconds"), ANSI_TO_TCHAR(__func__), *GetStateName(CurrentState), StateTimeouts[(uint8)CurrentState]);
+			EndICVFXTest();
+		}
+		else
+		{
+			// Treat all other timeouts as errors
+			UE_LOG(LogICVFXTest, Error, TEXT("AutoTest %s: '%s' test state timed out after %f seconds"), ANSI_TO_TCHAR(__func__), *GetStateName(CurrentState), StateTimeouts[(uint8)CurrentState]);
+			EndICVFXTest(1);
+		}
+	}
+}
+
+void UICVFXTestControllerAutoTest::BeginDestroy()
+{
+	UnbindAllDelegates();
+
+	for (int32 StateIdx = 0; StateIdx < (uint8)EICVFXAutoTestState::MAX; ++StateIdx)
+	{
+		delete States[StateIdx];
+		States[StateIdx] = nullptr;
+	}
+
+	Super::BeginDestroy();
+}
+
+void UICVFXTestControllerAutoTest::OnPreWorldInitialize(UWorld* const World)
+{
+	check(World);
+	World->GameStateSetEvent.AddUObject(this, &ThisClass::OnGameStateSet);
+	World->OnWorldBeginPlay.AddUObject(this, &ThisClass::OnWorldBeginPlay);
+}
+
+void UICVFXTestControllerAutoTest::OnGameStateSet(AGameStateBase* const GameStateBase)
+{
+	if (UWorld* const World = GetWorld())
+	{
+		World->GameStateSetEvent.RemoveAll(this);
+	}
+}
+
+void UICVFXTestControllerAutoTest::OnWorldBeginPlay()
+{
+	SetTestState(EICVFXAutoTestState::Soak);
+	
+	if (RequestsMemReport())
+	{
+		ExecuteMemReport();
+		SetMemReportTimer();
+	}
+}
+
+void UICVFXTestControllerAutoTest::OnSoakTimeChanged()
+{
+	const float SoakTime = ICVFXTest::CVarSoakTime->GetFloat();
+	StateTimeouts[(uint8)EICVFXAutoTestState::Soak] = SoakTime > 0.0f ? SoakTime + 15.0f : TNumericLimits<float>::Max();
+
+	if (CurrentState == EICVFXAutoTestState::Soak)
+	{
+		UE_LOG(LogICVFXTest, Display, TEXT("AutoTest Soak %s: soak time changed to %f seconds..."), ANSI_TO_TCHAR(__func__), ICVFXTest::CVarSoakTime->GetFloat());
+	}
+}
diff --git a/Engine/Plugins/VirtualProduction/ICVFXTesting/Source/ICVFXTesting/Private/ICVFXTestControllerBase.cpp b/Engine/Plugins/VirtualProduction/ICVFXTesting/Source/ICVFXTesting/Private/ICVFXTestControllerBase.cpp
new file mode 100644
--- /dev/null
+++ b/Engine/Plugins/VirtualProduction/ICVFXTesting/Source/ICVFXTesting/Private/ICVFXTestControllerBase.cpp
@@ -0,0 +1,326 @@
+// Copyright Epic Games, Inc. All Rights Reserved.
+
+#include "ICVFXTestControllerBase.h"
+
+#include "Engine/Engine.h"
+#include "GameFramework/PlayerController.h"
+#include "HAL/FileManager.h"
+#include "HAL/IConsoleManager.h"
+#include "Misc/CommandLine.h"
+#include "Misc/EngineVersion.h"
+#include "PlatformFeatures.h"
+#include "VideoRecordingSystem.h"
+
+#include "GauntletModule.h"
+
+DEFINE_LOG_CATEGORY(LogICVFXTest);
+
+namespace ICVFXTest
+{
+	static TAutoConsoleVariable<int32> CVarMaxRunCount(
+		TEXT("ICVFXTest.MaxRunCount"),
+		1,
+		TEXT("Max number of runs desired for this session."),
+		ECVF_Default);
+
+	static TAutoConsoleVariable<FString> CVarMemReportArgs(
+		TEXT("ICVFXTest.MemReportArgs"),
+		TEXT("-full -csv"),
+		TEXT("Arguments passed to the MemReport console command."),
+		ECVF_Default);
+
+	static TAutoConsoleVariable<float> CVarMemReportInterval(
+		TEXT("ICVFXTest.MemReportInterval"),
+		60.0f * 30.0f,
+		TEXT("Duration in seconds between MemReport calls during automated testing."),
+		ECVF_Default);
+}
+
+
+UICVFXTestControllerBase::UICVFXTestControllerBase(const FObjectInitializer& ObjectInitializer)
+	: Super(ObjectInitializer)
+	, WarnStuckTime(10 * 60)
+	, NextWarnStuckTime(0)
+	, RunCount(0)
+	, bRequestsFPSChart(false)
+	, bRequestsMemReport(false)
+	, bRequestsVideoCapture(false)
+{
+}
+
+
+void UICVFXTestControllerBase::OnInit()
+{
+	Super::OnInit();
+
+	int32 MaxRunCount;
+	if (FParse::Value(FCommandLine::Get(), TEXT("ICVFXTest.MaxRunCount"), MaxRunCount))
+	{
+		ICVFXTest::CVarMaxRunCount->Set(MaxRunCount);
+	}
+
+	bRequestsFPSChart = FParse::Param(FCommandLine::Get(), TEXT("ICVFXTest.FPSChart"));
+	bRequestsMemReport = FParse::Param(FCommandLine::Get(), TEXT("ICVFXTest.MemReport"));
+	bRequestsVideoCapture = FParse::Param(FCommandLine::Get(), TEXT("ICVFXTest.VideoCapture"));
+
+	MemReportTimerDelegate.BindUObject(this, &ThisClass::OnMemReportTimerExpired);
+
+	const FConsoleCommandDelegate MemReportIntervalDelegate = FConsoleCommandDelegate::CreateUObject(this, &ThisClass::OnMemReportIntervalChanged);
+	MemReportIntervalSink = IConsoleManager::Get().RegisterConsoleVariableSink_Handle(MemReportIntervalDelegate);
+
+	FWorldDelegates::OnPreWorldInitialization.AddUObject(this, &ThisClass::OnPreWorldInitializeInternal);
+}
+
+void UICVFXTestControllerBase::OnTick(float TimeDelta)
+{
+	if ((NextWarnStuckTime > 0) && (GetTimeInCurrentState() >= NextWarnStuckTime))
+	{
+		UE_LOG(LogICVFXTest, Display, TEXT("Have been in same state for %.02f mins"), GetTimeInCurrentState() / 60.0f);
+		NextWarnStuckTime += WarnStuckTime;
+	}
+}
+
+void UICVFXTestControllerBase::OnStateChange(FName OldState, FName NewState)
+{
+	if ((NewState != GetCurrentState()) && (NewState != FGauntletStates::Initialized))
+	{
+		NextWarnStuckTime = WarnStuckTime;
+	}
+}
+
+void UICVFXTestControllerBase::BeginDestroy()
+{
+	UnbindAllDelegates();
+	Super::BeginDestroy();
+}
+
+void UICVFXTestControllerBase::EndICVFXTest(const int32 ExitCode/*=0*/)
+{
+	UnbindAllDelegates();
+
+	UE_LOG(LogICVFXTest, Display, TEXT("%s: test completed, requesting exit..."), ANSI_TO_TCHAR(__func__));
+	EndTest(ExitCode);
+}
+
+void UICVFXTestControllerBase::UnbindAllDelegates()
+{
+	ClearMemReportTimer();
+	MemReportTimerDelegate.Unbind();
+	IConsoleManager::Get().UnregisterConsoleVariableSink_Handle(MemReportIntervalSink);
+	FWorldDelegates::OnPreWorldInitialization.RemoveAll(this);
+}
+
+uint32 UICVFXTestControllerBase::GetRunCount() const
+{
+	return RunCount;
+}
+
+uint32 UICVFXTestControllerBase::GetMaxRunCount() const
+{
+	const int32 MaxRunCount = ICVFXTest::CVarMaxRunCount->GetInt();
+	return MaxRunCount > 0 ? MaxRunCount : TNumericLimits<uint32>::Max();
+}
+
+uint32 UICVFXTestControllerBase::GetRunsRemaining() const
+{
+	return GetMaxRunCount() - GetRunCount();
+}
+
+uint32 UICVFXTestControllerBase::MarkRunComplete()
+{
+	++RunCount;
+	UE_LOG(LogICVFXTest, Display, TEXT("%s: run complete (%u/%u)..."), ANSI_TO_TCHAR(__func__), GetRunCount(), GetMaxRunCount());
+	return RunCount;
+}
+
+bool UICVFXTestControllerBase::RequestsFPSChart()
+{
+	return bRequestsFPSChart;
+}
+
+bool UICVFXTestControllerBase::RequestsMemReport()
+{
+	return bRequestsMemReport;
+}
+
+bool UICVFXTestControllerBase::RequestsVideoCapture()
+{
+	return bRequestsVideoCapture;
+}
+
+void UICVFXTestControllerBase::ConsoleCommand(const TCHAR* Cmd)
+{
+	if (APlayerController* Controller = GetFirstPlayerController())
+	{
+		Controller->ConsoleCommand(Cmd);
+		UE_LOG(LogICVFXTest, Display, TEXT("Issued console command: '%s'"), Cmd);
+	}
+}
+
+void UICVFXTestControllerBase::ExecuteMemReport(const TOptional<FString> Args/*=TOptional<FString>()*/)
+{
+	ConsoleCommand(*FString(TEXT("MemReport ") + (Args.IsSet() ? Args.GetValue() : ICVFXTest::CVarMemReportArgs->GetString())));
+}
+
+void UICVFXTestControllerBase::SetMemReportTimer(const TOptional<float> Interval/*=TOptional<float>()*/)
+{
+	if (UWorld* const World = GetWorld())
+	{
+		FTimerManager& TimerManager = World->GetTimerManager();
+
+		const float AdjustedInterval = Interval.IsSet() ? Interval.GetValue() : ICVFXTest::CVarMemReportInterval->GetFloat();
+		if (AdjustedInterval > 0.0f)
+		{
+			UE_LOG(LogICVFXTest, Display, TEXT("%s: interval is set to %f seconds..."), ANSI_TO_TCHAR(__func__), AdjustedInterval);
+			TimerManager.SetTimer(MemReportTimerHandle, MemReportTimerDelegate, AdjustedInterval, true);
+		}
+		else if (MemReportTimerHandle.IsValid())
+		{
+			UE_LOG(LogICVFXTest, Display, TEXT("%s: interval is set to 0, pausing timer..."), ANSI_TO_TCHAR(__func__));
+			TimerManager.PauseTimer(MemReportTimerHandle);
+		}
+		else
+		{
+			UE_LOG(LogICVFXTest, Warning, TEXT("%s: interval is set to 0 and no timer is set, skipping..."), ANSI_TO_TCHAR(__func__));
+		}
+	}
+}
+
+void UICVFXTestControllerBase::ClearMemReportTimer()
+{
+	if (UWorld* const World = GetWorld())
+	{
+		if (MemReportTimerHandle.IsValid())
+		{
+			World->GetTimerManager().ClearTimer(MemReportTimerHandle);
+			UE_LOG(LogICVFXTest, Display, TEXT("%s: cleared..."), ANSI_TO_TCHAR(__func__));
+		}
+	}
+}
+
+bool UICVFXTestControllerBase::TryStartingVideoCapture()
+{
+	if (IVideoRecordingSystem* const VideoRecordingSystem = IPlatformFeaturesModule::Get().GetVideoRecordingSystem())
+	{
+		const EVideoRecordingState RecordingState = VideoRecordingSystem->GetRecordingState();
+
+		if (RecordingState == EVideoRecordingState::None)
+		{
+			VideoRecordingSystem->EnableRecording(true);
+
+			const FString Suffix = FString::Printf(TEXT("_%s"), *FDateTime::Now().ToString(TEXT("%H-%M-%S")));
+			const FString VideoFilename = FString::Printf(TEXT("AutoTest_CL-%d%s"), FEngineVersion::Current().GetChangelist(), *Suffix);
+			VideoRecordingTitle = FText::FromString(VideoFilename);
+			constexpr bool bInAutoStart = true;
+			constexpr bool bInAutoContinue = false;
+			constexpr bool bInExportToLibrary = false;
+			FVideoRecordingParameters VideoRecordingParameters(VideoRecordingSystem->GetMaximumRecordingSeconds(), bInAutoStart, bInAutoContinue, bInExportToLibrary, FPlatformMisc::GetPlatformUserForUserIndex(0));
+			VideoRecordingSystem->NewRecording(*VideoFilename, VideoRecordingParameters);
+
+			if (VideoRecordingSystem->IsEnabled())
+			{
+				if (RecordingState == EVideoRecordingState::Starting || RecordingState == EVideoRecordingState::Recording)
+				{
+					UE_LOG(LogICVFXTest, Log, TEXT("AutoTest %s: starting video recording %s..."), ANSI_TO_TCHAR(__func__), *VideoFilename);
+					return true;
+				}
+				else
+				{
+					UE_LOG(LogICVFXTest, Warning, TEXT("AutoTest %s: failed to start video recording %s."), ANSI_TO_TCHAR(__func__), *VideoFilename);
+				}
+			}
+			else
+			{
+				UE_LOG(LogICVFXTest, Warning, TEXT("AutoTest %s: video recording could not be enabled."), ANSI_TO_TCHAR(__func__));
+			}
+		}
+		else
+		{
+			UE_LOG(LogICVFXTest, Warning, TEXT("AutoTest %s: could not start a new recording, may be already recording."), ANSI_TO_TCHAR(__func__));
+		}
+	}
+	else
+	{
+		UE_LOG(LogICVFXTest, Warning, TEXT("AutoTest %s: video recording system is null."), ANSI_TO_TCHAR(__func__));
+	}
+
+	return false;
+}
+
+bool UICVFXTestControllerBase::TryFinalizingVideoCapture(const bool bStopAutoContinue/*=false*/)
+{
+	if (IVideoRecordingSystem* const VideoRecordingSystem = IPlatformFeaturesModule::Get().GetVideoRecordingSystem())
+	{
+		if (VideoRecordingSystem->GetRecordingState() != EVideoRecordingState::None)
+		{
+			VideoRecordingSystem->FinalizeRecording(true, VideoRecordingTitle, FText::GetEmpty(), bStopAutoContinue);
+
+			if (VideoRecordingSystem->GetRecordingState() == EVideoRecordingState::Finalizing)
+			{
+				UE_LOG(LogICVFXTest, Log, TEXT("AutoTest %s: finalizing recording..."), ANSI_TO_TCHAR(__func__));
+				return true;
+			}
+		}
+	}
+	else
+	{
+		UE_LOG(LogICVFXTest, Warning, TEXT("AutoTest %s: video recording system is null."), ANSI_TO_TCHAR(__func__));
+	}
+
+	return false;
+}
+
+void UICVFXTestControllerBase::OnPreWorldInitializeInternal(UWorld* World, const UWorld::InitializationValues IVS)
+{
+	TryEarlyExec(World);
+	OnPreWorldInitialize(World);
+}
+
+void UICVFXTestControllerBase::TryEarlyExec(UWorld* const World)
+{
+	check(World);
+
+	if (GEngine)
+	{
+		// Search the list of deferred commands
+		const TArray<FString>& DeferredCmds = GEngine->DeferredCommands;
+		TArray<int32> ExecutedIndices;
+		for (int32 DeferredCmdIndex = 0; DeferredCmdIndex < DeferredCmds.Num(); ++DeferredCmdIndex)
+		{
+			// If the deferred command is one that should be executed early
+			const FString& DeferredCmd = DeferredCmds[DeferredCmdIndex];
+			if (CmdsToExecEarly.ContainsByPredicate([&DeferredCmd](const FString& CmdToFind) { return DeferredCmd.StartsWith(CmdToFind); }))
+			{
+				UE_LOG(LogICVFXTest, Display, TEXT("%s: executing '%s' early."), ANSI_TO_TCHAR(__func__), *DeferredCmd);
+				GEngine->Exec(World, *DeferredCmd);
+				ExecutedIndices.Push(DeferredCmdIndex);
+			}
+		}
+
+		// Remove the executed commands from the list of deferred commands
+		// Note: This is done in reverse order to ensure the cached indices remain valid
+		while (!ExecutedIndices.IsEmpty())
+		{
+			GEngine->DeferredCommands.RemoveAt(ExecutedIndices.Pop());
+		}
+	}
+}
+
+void UICVFXTestControllerBase::OnMemReportTimerExpired()
+{
+	ExecuteMemReport();
+}
+
+void UICVFXTestControllerBase::OnMemReportIntervalChanged()
+{
+	if (MemReportTimerHandle.IsValid())
+	{
+		if (UWorld* const World = GetWorld())
+		{
+			if (World->GetTimerManager().GetTimerRate(MemReportTimerHandle) != ICVFXTest::CVarMemReportInterval->GetFloat())
+			{
+				SetMemReportTimer();
+			}
+		}
+	}
+}
diff --git a/Engine/Plugins/VirtualProduction/ICVFXTesting/Source/ICVFXTesting/Private/ICVFXTestingModule.cpp b/Engine/Plugins/VirtualProduction/ICVFXTesting/Source/ICVFXTesting/Private/ICVFXTestingModule.cpp
new file mode 100644
--- /dev/null
+++ b/Engine/Plugins/VirtualProduction/ICVFXTesting/Source/ICVFXTesting/Private/ICVFXTestingModule.cpp
@@ -0,0 +1,6 @@
+// Copyright Epic Games, Inc. All Rights Reserved.
+
+#include "Modules/ModuleManager.h"
+#include "ICVFXTestControllerAutoTest.h"
+
+IMPLEMENT_MODULE(FDefaultModuleImpl, ICVFXTestingModule);
diff --git a/Engine/Plugins/VirtualProduction/ICVFXTesting/Source/ICVFXTesting/Public/ICVFXTestControllerAutoTest.h b/Engine/Plugins/VirtualProduction/ICVFXTesting/Source/ICVFXTesting/Public/ICVFXTestControllerAutoTest.h
new file mode 100644
--- /dev/null
+++ b/Engine/Plugins/VirtualProduction/ICVFXTesting/Source/ICVFXTesting/Public/ICVFXTestControllerAutoTest.h
@@ -0,0 +1,85 @@
+// Copyright Epic Games, Inc. All Rights Reserved.
+#pragma once
+
+#include "ICVFXTestControllerBase.h"
+
+#include "HAL/IConsoleManager.h"
+#include "ICVFXTestControllerAutoTest.generated.h"
+
+class ULocalPlayer;
+
+UENUM()
+enum class EICVFXAutoTestState : uint8
+{
+	InitialLoad,
+	Soak,
+	Idle ,
+	Finished,
+	Shutdown,
+
+	MAX
+};
+
+class FICVFXAutoTestState
+{
+public:
+	using Super = FICVFXAutoTestState;
+
+	FICVFXAutoTestState() = delete;
+	FICVFXAutoTestState(class UICVFXTestControllerAutoTest* const TestController) : Controller(TestController) {}
+	virtual ~FICVFXAutoTestState() {}
+
+	virtual void Start(const EICVFXAutoTestState PrevState);
+	virtual void End(const EICVFXAutoTestState NewState) {}
+	virtual void Tick(const float TimeDelta);
+
+	double GetTestStateTime() const 
+	{
+		return TimeSinceStart;
+	}
+
+protected:
+	class UICVFXTestControllerAutoTest* Controller;
+
+private:
+	double TimeSinceStart = 0.0;
+};
+
+UCLASS()
+class UICVFXTestControllerAutoTest : public UICVFXTestControllerBase
+{
+	GENERATED_BODY()
+
+public:
+	FString GetStateName(const EICVFXAutoTestState State) const;
+	FICVFXAutoTestState& GetTestState() const;
+	FICVFXAutoTestState& SetTestState(const EICVFXAutoTestState NewState);
+
+	virtual void EndICVFXTest(const int32 ExitCode=0) override;
+
+protected:
+	virtual void OnInit() override;
+	virtual void OnPreMapChange() override;
+	virtual void OnTick(float TimeDelta) override;
+	virtual void BeginDestroy() override;
+
+private:
+	virtual void UnbindAllDelegates() override;
+
+	FICVFXAutoTestState* States[(uint8)EICVFXAutoTestState::MAX];
+	float StateTimeouts[(uint8)EICVFXAutoTestState::MAX];
+	EICVFXAutoTestState CurrentState;
+
+	virtual void OnPreWorldInitialize(UWorld* World) override;
+
+	UFUNCTION()
+	void OnWorldBeginPlay();
+
+	UFUNCTION()
+	void OnGameStateSet(AGameStateBase* const GameStateBase);
+	
+	FConsoleVariableSinkHandle SoakTimeSink;
+
+	UFUNCTION()
+	void OnSoakTimeChanged();
+};
\ No newline at end of file
diff --git a/Engine/Plugins/VirtualProduction/ICVFXTesting/Source/ICVFXTesting/Public/ICVFXTestControllerBase.h b/Engine/Plugins/VirtualProduction/ICVFXTesting/Source/ICVFXTesting/Public/ICVFXTestControllerBase.h
new file mode 100644
--- /dev/null
+++ b/Engine/Plugins/VirtualProduction/ICVFXTesting/Source/ICVFXTesting/Public/ICVFXTestControllerBase.h
@@ -0,0 +1,99 @@
+// Copyright Epic Games, Inc. All Rights Reserved.
+
+#pragma once
+
+#include "GauntletTestController.h"
+
+#include "Engine/TimerHandle.h"
+#include "Engine/World.h"
+#include "HAL/IConsoleManager.h"
+#include "TimerManager.h"
+
+#include "ICVFXTestControllerBase.generated.h"
+
+
+DECLARE_LOG_CATEGORY_EXTERN(LogICVFXTest, Log, All)
+
+class AActor;
+class AGameState;
+class IConsoleVariable;
+
+UCLASS()
+class UICVFXTestControllerBase : public UGauntletTestController
+{
+	GENERATED_BODY()
+
+public:
+	UICVFXTestControllerBase(const FObjectInitializer& ObjectInitializer);
+
+protected:
+	// ~Begin UGauntletTestController Interface
+	virtual void OnInit() override;
+	virtual void OnTick(float TimeDelta) override;
+	virtual void OnStateChange(FName OldState, FName NewState) override;
+	virtual void BeginDestroy() override;
+	// ~End UGauntletTestController Interface
+
+	virtual void EndICVFXTest(const int32 ExitCode=0);
+	virtual void UnbindAllDelegates();
+
+public:
+	uint32 GetRunCount() const;
+	uint32 GetMaxRunCount() const;
+	uint32 GetRunsRemaining() const;
+	uint32 MarkRunComplete();
+
+	bool RequestsFPSChart();
+	bool RequestsMemReport();
+	bool RequestsVideoCapture();
+
+	void ConsoleCommand(const TCHAR* Cmd);
+
+	void ExecuteMemReport(const TOptional<FString> Args=TOptional<FString>());
+	void SetMemReportTimer(const TOptional<float> Interval=TOptional<float>());
+	void ClearMemReportTimer();
+
+	bool TryStartingVideoCapture();
+	bool TryFinalizingVideoCapture(const bool bStopAutoContinue = false);
+
+private:
+	float WarnStuckTime;
+	float NextWarnStuckTime;
+
+	void OnPreWorldInitializeInternal(UWorld* World, const UWorld::InitializationValues IVS);
+	virtual void OnPreWorldInitialize(UWorld* World) {}
+
+	UFUNCTION()
+	void TryEarlyExec(UWorld* const World);
+
+	/** List of command prefixes used to find and execute commands passed in through commandline or config OnInit. */
+	const TArray<FString> CmdsToExecEarly = {
+		TEXT("t.FPSChart.DoCsvProfile"),
+		TEXT("ICVFXTest"),
+		TEXT("r.nanite"),
+		TEXT("r.ScreenPercentage"),
+		TEXT("r.RayTracing"),
+		TEXT("r.DynamicGlobalIlluminationMethod"),
+		TEXT("r.ReflectionMethod"),
+		TEXT("r.Lumen"),
+		TEXT("FX.AllowGPUParticles"),
+		TEXT("r.Shadow.Virtual.Enable")};
+
+	uint32 RunCount;
+
+	bool bRequestsFPSChart;
+	bool bRequestsMemReport;
+	bool bRequestsVideoCapture;
+
+	FTimerHandle MemReportTimerHandle;
+	FTimerDelegate MemReportTimerDelegate;
+	FConsoleVariableSinkHandle MemReportIntervalSink;
+
+	UFUNCTION()
+	void OnMemReportTimerExpired();
+
+	UFUNCTION()
+	void OnMemReportIntervalChanged();
+
+	FText VideoRecordingTitle;
+};
\ No newline at end of file
