Rename ICVFX Testing module to fix CIS

diff --git a/Engine/Plugins/VirtualProduction/ICVFXTesting/Source/ICVFXTesting/Private/ICVFXTestingModule.cpp b/Engine/Plugins/VirtualProduction/ICVFXTesting/Source/ICVFXTesting/Private/ICVFXTestingModule.cpp
--- a/Engine/Plugins/VirtualProduction/ICVFXTesting/Source/ICVFXTesting/Private/ICVFXTestingModule.cpp
+++ b/Engine/Plugins/VirtualProduction/ICVFXTesting/Source/ICVFXTesting/Private/ICVFXTestingModule.cpp
@@ -3,4 +3,4 @@
 #include "Modules/ModuleManager.h"
 #include "ICVFXTestControllerAutoTest.h"
 
-IMPLEMENT_MODULE(FDefaultModuleImpl, ICVFXTestingModule);
+IMPLEMENT_MODULE(FDefaultModuleImpl, ICVFXTesting);
