Fix CIS when compiling ICVFX testing plugin for Wingdk

diff --git a/Engine/Plugins/VirtualProduction/ICVFXTesting/ICVFXTesting.uplugin b/Engine/Plugins/VirtualProduction/ICVFXTesting/ICVFXTesting.uplugin
--- a/Engine/Plugins/VirtualProduction/ICVFXTesting/ICVFXTesting.uplugin
+++ b/Engine/Plugins/VirtualProduction/ICVFXTesting/ICVFXTesting.uplugin
@@ -32,7 +32,8 @@
 		{
 			"Name": "ICVFXTesting",
 			"Type": "Runtime",
-			"LoadingPhase": "Default"
+			"LoadingPhase": "Default",
+			"PlatformAllowList": [ "Win64" ]
 		}
 	]
 }
