[nDisplay] Fixed PerfTesting reports generation after 5.7 gauntlet changes
- Fixed reports generation for local/on-stage builds only
- Fixed crash when test DCRA doesn't have any ICVFX camera components
- Moved some test code from .h to .cpp

#rnx

#Synced-CL 0

diff --git a/Engine/Plugins/VirtualProduction/ICVFXTesting/Scripts/GauntletSettings.xml b/Engine/Plugins/VirtualProduction/ICVFXTesting/Scripts/GauntletSettings.xml
--- a/Engine/Plugins/VirtualProduction/ICVFXTesting/Scripts/GauntletSettings.xml
+++ b/Engine/Plugins/VirtualProduction/ICVFXTesting/Scripts/GauntletSettings.xml
@@ -99,6 +99,10 @@
 	<Property Name="ExtraGauntletCommonArgs" Value="-screenshotperiod=0 -ICVFXTest.StatCommands -ICVFXTest.MaxRunCount=$(MaxRunCount) $(LogOutputParams)" />
 	<Property Name="ExtraGauntletEditorTestArgs" Value="-game" />
 
+	<Do If="$(IsOnStagePerfTest)">
+		<Property Name="ExtraGauntletEditorTestArgs" Value="$(ExtraGauntletEditorTestArgs) -RetainDeviceArtifacts" />
+	</Do>
+
 	<Property Name="ExtraGauntletEditorTestArgs" Value="$(ExtraGauntletEditorTestArgs) -ICVFXTest.IsBuildMachine=$(IsBuildMachine)" />
 	<Property Name="ExtraGauntletEditorTestArgs" Value="$(ExtraGauntletEditorTestArgs) -ICVFXTest.TestName=$(ProjectTestName)" />
 	<Property Name="ExtraGauntletEditorTestArgs" Value="$(ExtraGauntletEditorTestArgs) -ICVFXTest.MapOverride=$(MapOverride)" />
diff --git a/Engine/Plugins/VirtualProduction/ICVFXTesting/Source/ICVFXTesting/Private/ICVFXTestControllerAutoTest.cpp b/Engine/Plugins/VirtualProduction/ICVFXTesting/Source/ICVFXTesting/Private/ICVFXTestControllerAutoTest.cpp
--- a/Engine/Plugins/VirtualProduction/ICVFXTesting/Source/ICVFXTesting/Private/ICVFXTestControllerAutoTest.cpp
+++ b/Engine/Plugins/VirtualProduction/ICVFXTesting/Source/ICVFXTesting/Private/ICVFXTestControllerAutoTest.cpp
@@ -17,10 +17,12 @@
 #include "Logging/LogVerbosity.h"
 #include "PlatformFeatures.h"
 #include "ProfilingDebugging/CsvProfiler.h"
+#include "Render/Viewport/IDisplayClusterViewportManager.h"
 #include "VideoRecordingSystem.h"
 #include "Roles/LiveLinkTransformRole.h"
 #include "UObject/SoftObjectPtr.h"
 
+
 CSV_DEFINE_CATEGORY(ICVFXTest, true);
 
 namespace ICVFXTest
@@ -413,6 +415,127 @@ void UICVFXTestControllerAutoTest::SetTestLocations(const TArray<AActor*> InTest
 	StateTimeouts[(uint8)EICVFXAutoTestState::TraverseTestLocations] = TestLocations.Num() * TimePerTestLocation;
 }
 
+void UICVFXTestControllerAutoTest::InitializeLiveLink()
+{
+	if (!IsValid(DisplayClusterActor))
+	{
+		return;
+	}
+
+	// LL preset is located in the same (ICVFXTesting) plugin, therefore must always exist
+	ULiveLinkPreset* LLPreset = LoadObject<ULiveLinkPreset>(NULL, TEXT("/ICVFXTesting/AutomationPreset"));
+	if (!IsValid(LLPreset))
+	{
+		check(false);
+		return;
+	}
+
+	UDisplayClusterICVFXCameraComponent* const Component = DisplayClusterActor->FindComponentByClass<UDisplayClusterICVFXCameraComponent>();
+	if (!IsValid(Component))
+	{
+		return;
+	}
+
+	const FDisplayClusterConfigurationICVFX_CameraSettings& CameraSettings = Component->GetCameraSettingsICVFX();
+	ACineCameraActor* const CameraActor = CameraSettings.ExternalCameraActor.Get();
+	if (!IsValid(CameraActor))
+	{
+		return;
+	}
+
+	ULiveLinkComponentController* LiveLinkComponentController = NewObject<ULiveLinkComponentController>(CameraActor);
+	LiveLinkComponentController->bUpdateInEditor = true;
+	LiveLinkComponentController->bEvaluateLiveLink = true;
+
+	CameraActor->AddInstanceComponent(LiveLinkComponentController);
+	LiveLinkComponentController->RegisterComponentWithWorld(CameraActor->GetWorld());
+
+	FName LLVirtualSubjectName;
+	if (LLPreset->GetSubjectPresets().Num())
+	{
+		LLVirtualSubjectName = LLPreset->GetSubjectPresets()[0].Key.SubjectName.Name;
+	}
+
+	LLPreset->ApplyToClientLatent([LiveLinkComponentController, LLVirtualSubjectName, CameraActor](bool bSuccess)
+		{
+			if (!IsValid(LiveLinkComponentController) || !IsValid(CameraActor))
+			{
+				return;
+			}
+
+			if (bSuccess)
+			{
+				FLiveLinkSubjectRepresentation SubjectRepresentation;
+				SubjectRepresentation.Role = ULiveLinkTransformRole::StaticClass();
+				SubjectRepresentation.Subject = LLVirtualSubjectName;
+
+				LiveLinkComponentController->SetSubjectRepresentation(SubjectRepresentation);
+				LiveLinkComponentController->SetControlledComponent(ULiveLinkTransformRole::StaticClass(), CameraActor->GetCineCameraComponent());
+
+				if (TObjectPtr<ULiveLinkControllerBase>* Controller = LiveLinkComponentController->ControllerMap.Find(ULiveLinkTransformRole::StaticClass()))
+				{
+					if (ULiveLinkTransformController* TransformController = Cast<ULiveLinkTransformController>(*Controller))
+					{
+						TransformController->TransformData.bUseLocation = false;
+					}
+				}
+			}
+		});
+}
+
+void UICVFXTestControllerAutoTest::UpdateInnerGPUIndex()
+{
+	TArray<UDisplayClusterICVFXCameraComponent*> CameraComponents;
+	ADisplayClusterRootActor* RootActor = Cast<ADisplayClusterRootActor>(DisplayClusterActor);
+	if (!IsValid(RootActor))
+	{
+		return;
+	}
+
+	constexpr bool bIncludeChildActors = false;
+	RootActor->GetComponents<UDisplayClusterICVFXCameraComponent>(CameraComponents, bIncludeChildActors);
+
+	UWorld* CurrentWorld = RootActor->GetWorld();
+
+	for (UDisplayClusterICVFXCameraComponent* CameraComponent : CameraComponents)
+	{
+		CameraComponent->CameraSettings.RenderSettings.AdvancedRenderSettings.GPUIndex = GetInnerGPUIndex();
+	}
+
+	if (IDisplayClusterViewportManager* ViewportManager = RootActor->GetViewportManager())
+	{
+		const FString NodeId = IDisplayCluster::Get().GetClusterMgr()->GetNodeId();
+		const EDisplayClusterRenderFrameMode RenderMode = EDisplayClusterRenderFrameMode::Mono;
+
+		ViewportManager->GetConfiguration().UpdateConfigurationForClusterNode(RenderMode, CurrentWorld, NodeId);
+	}
+}
+
+void UICVFXTestControllerAutoTest::UpdateTestLocations()
+{
+	UClass* TestLocationClass = AICVFXTestLocation::StaticClass();
+
+	TestLocations.Reset();
+
+	for (TActorIterator<AActor> ItActor = TActorIterator<AActor>(GetWorld(), TestLocationClass); ItActor; ++ItActor)
+	{
+		TestLocations.Add(*ItActor);
+	}
+
+	if (TestLocations.Num())
+	{
+		UE_LOG(LogICVFXTest, Display, TEXT("Found %d test locations."), TestLocations.Num());
+	}
+	else
+	{
+		UE_LOG(LogICVFXTest, Display, TEXT("Could not find test locations, defaulting to display cluster position."));
+		// Todo: If no uobject path passed to us, default on the first display cluster actor we find?
+		TestLocations.Add(DisplayClusterActor);
+	}
+
+	SetTestLocations(TestLocations);
+}
+
 void UICVFXTestControllerAutoTest::EndICVFXTest(const int32 ExitCode /*= 0*/)
 {
 	if (CurrentState != EICVFXAutoTestState::Shutdown)
diff --git a/Engine/Plugins/VirtualProduction/ICVFXTesting/Source/ICVFXTesting/Public/ICVFXTestControllerAutoTest.h b/Engine/Plugins/VirtualProduction/ICVFXTesting/Source/ICVFXTesting/Public/ICVFXTestControllerAutoTest.h
--- a/Engine/Plugins/VirtualProduction/ICVFXTesting/Source/ICVFXTesting/Public/ICVFXTestControllerAutoTest.h
+++ b/Engine/Plugins/VirtualProduction/ICVFXTesting/Source/ICVFXTesting/Public/ICVFXTestControllerAutoTest.h
@@ -17,7 +17,6 @@
 #include "LiveLinkPreset.h"
 #include "LiveLinkPresetTypes.h"
 #include "Roles/LiveLinkTransformRole.h"
-#include "Render/Viewport/IDisplayClusterViewportManager.h"
 
 #include "ICVFXTestControllerAutoTest.generated.h"
 
@@ -28,7 +27,7 @@ enum class EICVFXAutoTestState : uint8
 {
 	InitialLoad,
 	Soak,
-	TraverseTestLocations ,
+	TraverseTestLocations,
 	Finished,
 	Shutdown,
 
@@ -105,101 +104,9 @@ public:
 		return InnerGPUIndex.load();
 	}
 
-	void InitializeLiveLink() const
-	{
-		if (DisplayClusterActor)
-		{
-			ULiveLinkPreset* LLPreset = LoadObject<ULiveLinkPreset>(NULL, TEXT("/ICVFXTesting/AutomationPreset"));
-			check(LLPreset);
-
-			UDisplayClusterICVFXCameraComponent* Component = DisplayClusterActor->FindComponentByClass<UDisplayClusterICVFXCameraComponent>();
-			if (ACineCameraActor* CameraActor = Component->GetCameraSettingsICVFX().ExternalCameraActor.Get())
-			{
-				ULiveLinkComponentController* LiveLinkComponentController = NewObject<ULiveLinkComponentController>(CameraActor);
-				LiveLinkComponentController->bUpdateInEditor = true;
-				LiveLinkComponentController->bEvaluateLiveLink = true;
-
-
-				CameraActor->AddInstanceComponent(LiveLinkComponentController);
-				LiveLinkComponentController->RegisterComponentWithWorld(CameraActor->GetWorld());
-
-				FName LLVirtualSubjectName;
-				if (LLPreset->GetSubjectPresets().Num())
-				{
-					LLVirtualSubjectName = LLPreset->GetSubjectPresets()[0].Key.SubjectName.Name;
-				}
-
-				LLPreset->ApplyToClientLatent([LiveLinkComponentController, LLVirtualSubjectName, CameraActor](bool bSuccess)
-				{
-					if (bSuccess)
-					{
-						FLiveLinkSubjectRepresentation SubjectRepresentation;
-						SubjectRepresentation.Role = ULiveLinkTransformRole::StaticClass();
-						SubjectRepresentation.Subject = LLVirtualSubjectName;
-
-						LiveLinkComponentController->SetSubjectRepresentation(SubjectRepresentation);
-						LiveLinkComponentController->SetControlledComponent(ULiveLinkTransformRole::StaticClass(), CameraActor->GetCineCameraComponent());
-
-						if (TObjectPtr<ULiveLinkControllerBase>* Controller = LiveLinkComponentController->ControllerMap.Find(ULiveLinkTransformRole::StaticClass()))
-						{
-							if (ULiveLinkTransformController* TransformController = Cast<ULiveLinkTransformController>(*Controller))
-							{
-								TransformController->TransformData.bUseLocation = false;
-							}
-						}
-					}
-				});
-			}
-		}
-	}
-
-	void UpdateInnerGPUIndex()
-	{
-		TArray<UDisplayClusterICVFXCameraComponent*> CameraComponents;
-		ADisplayClusterRootActor* RootActor = Cast<ADisplayClusterRootActor>(DisplayClusterActor);
-		constexpr bool bIncludeChildActors = false;
-		RootActor->GetComponents<UDisplayClusterICVFXCameraComponent>(CameraComponents, bIncludeChildActors);
-
-		UWorld* CurrentWorld = RootActor->GetWorld();
-
-		for (UDisplayClusterICVFXCameraComponent* CameraComponent : CameraComponents)
-		{
-			CameraComponent->CameraSettings.RenderSettings.AdvancedRenderSettings.GPUIndex = GetInnerGPUIndex();
-		}
-
-		if (IDisplayClusterViewportManager* ViewportManager = RootActor->GetViewportManager())
-		{
-			const FString NodeId = IDisplayCluster::Get().GetClusterMgr()->GetNodeId();
-			const EDisplayClusterRenderFrameMode RenderMode = EDisplayClusterRenderFrameMode::Mono;
-
-			ViewportManager->GetConfiguration().UpdateConfigurationForClusterNode(RenderMode, CurrentWorld, NodeId);
-		}
-	}
-
-	void UpdateTestLocations()
-	{
-		UClass* TestLocationClass = AICVFXTestLocation::StaticClass();
-
-		TestLocations.Reset();
-
-		for (TActorIterator<AActor> ItActor = TActorIterator<AActor>(GetWorld(), TestLocationClass); ItActor; ++ItActor)
-		{
-			TestLocations.Add(*ItActor);
-		}
-
-		if (TestLocations.Num())
-		{
-			UE_LOG(LogICVFXTest, Display, TEXT("Found %d test locations."), TestLocations.Num());
-		}
-		else
-		{
-			UE_LOG(LogICVFXTest, Display, TEXT("Could not find test locations, defaulting to display cluster position."));
-			// Todo: If no uobject path passed to us, default on the first display cluster actor we find?
-			TestLocations.Add(DisplayClusterActor);
-		}
-
-		SetTestLocations(TestLocations);
-	}
+	void InitializeLiveLink();
+	void UpdateInnerGPUIndex();
+	void UpdateTestLocations();
 
 protected:
 	virtual void OnInit() override;
@@ -210,7 +117,8 @@ protected:
 public:
 	float TimePerTestLocation = 60.f;
 	double TimeAtTestLocation = 0.0;
-	AActor* DisplayClusterActor = nullptr;
+	TObjectPtr<AActor> DisplayClusterActor = nullptr;
+
 private:
 	virtual void UnbindAllDelegates() override;
 
@@ -237,4 +145,4 @@ private:
 
 	/** What GPU Index should be used for the inner viewport. */
 	std::atomic<int32> InnerGPUIndex = 0;
-};
\ No newline at end of file
+};
